Avoid looking for tests in pluting content python site-packages folder

#rnx

diff --git a/Engine/Plugins/Tests/PythonAutomationTest/Source/PythonAutomationTest/Private/PythonAutomationTest.cpp b/Engine/Plugins/Tests/PythonAutomationTest/Source/PythonAutomationTest/Private/PythonAutomationTest.cpp
--- a/Engine/Plugins/Tests/PythonAutomationTest/Source/PythonAutomationTest/Private/PythonAutomationTest.cpp
+++ b/Engine/Plugins/Tests/PythonAutomationTest/Source/PythonAutomationTest/Private/PythonAutomationTest.cpp
@@ -123,6 +123,11 @@ protected:
 		// Scan all the found files, use only test_*.py file
 		for (const FString& Filename : FilesInDirectory)
 		{
+			// Skip files from site-packages
+			if (Filename.Contains(TEXT("/site-packages/")))
+			{
+				continue;
+			}
 			OutBeautifiedNames.Add(ModuleName + TEXT(".") + BeautifyPath(Filename));
 			OutFileNames.Add(Filename);
 		}
